Fix dependency on engine for WebBrowser
#ue4

diff --git a/Engine/Source/Runtime/WebBrowser/Private/WebBrowserTexture.cpp b/Engine/Source/Runtime/WebBrowser/Private/WebBrowserTexture.cpp
deleted file mode 100644
diff --git a/Engine/Source/Runtime/WebBrowser/Private/WebBrowserTextureResource.cpp b/Engine/Source/Runtime/WebBrowser/Private/WebBrowserTextureResource.cpp
deleted file mode 100644
diff --git a/Engine/Source/Runtime/WebBrowser/Public/WebBrowserTexture.h b/Engine/Source/Runtime/WebBrowser/Public/WebBrowserTexture.h
deleted file mode 100644
diff --git a/Engine/Source/Runtime/WebBrowser/Public/WebBrowserTextureResource.h b/Engine/Source/Runtime/WebBrowser/Public/WebBrowserTextureResource.h
deleted file mode 100644
diff --git a/Engine/Source/Runtime/WebBrowser/Public/WebBrowserTextureSample.h b/Engine/Source/Runtime/WebBrowser/Public/WebBrowserTextureSample.h
deleted file mode 100644
diff --git a/Engine/Source/Runtime/WebBrowser/WebBrowser.Build.cs b/Engine/Source/Runtime/WebBrowser/WebBrowser.Build.cs
--- a/Engine/Source/Runtime/WebBrowser/WebBrowser.Build.cs
+++ b/Engine/Source/Runtime/WebBrowser/WebBrowser.Build.cs
@@ -9,48 +9,36 @@ public class WebBrowser : ModuleRules
 	{
 		PublicIncludePaths.Add("Runtime/WebBrowser/Public");
 		PrivateIncludePaths.Add("Runtime/WebBrowser/Private");
-        PublicIncludePaths.AddRange(
-            new string[] {
-                "Runtime/MediaUtils/Public",
-            });
-
-        PrivateIncludePaths.AddRange(
-            new string[] {
-                "Runtime/MediaUtils/Private",
-            });
-
         PrivateDependencyModuleNames.AddRange(
 			new string[]
 			{
 				"Core",
 				"CoreUObject",
 				"ApplicationCore",
-                "RHI",
+				"RHI",
 				"InputCore",
 				"Slate",
 				"SlateCore",
 				"Serialization",
-                "MediaUtils",
-                "RenderCore",
-                "Engine",
             }
         );
 
-        PrivateDependencyModuleNames.AddRange(
-        new string[] {
-                    "RenderCore",
-                    "RHI",
-                    "ShaderCore",
-                    "UtilityShaders",
-        });
-
         if (Target.Platform == UnrealTargetPlatform.Android)
 		{
-            // We need this one on Android for URL decoding
-            PrivateDependencyModuleNames.Add("HTTP");
-        }
+			// We need these on Android for external texture support
+			PrivateDependencyModuleNames.AddRange(
+				new string[]
+				{
+					"WebBrowserTexture",
+					"Engine",
+				}
+			);
+
+			// We need this one on Android for URL decoding
+			PrivateDependencyModuleNames.Add("HTTP");
+		}
 
-        if (Target.Platform == UnrealTargetPlatform.Win64
+		if (Target.Platform == UnrealTargetPlatform.Win64
 		||  Target.Platform == UnrealTargetPlatform.Win32
 		||  Target.Platform == UnrealTargetPlatform.Mac
 		||  Target.Platform == UnrealTargetPlatform.Linux)
diff --git a/Engine/Source/Runtime/WebBrowserTexture/Private/WebBrowserTexture.cpp b/Engine/Source/Runtime/WebBrowserTexture/Private/WebBrowserTexture.cpp
new file mode 100644
--- /dev/null
+++ b/Engine/Source/Runtime/WebBrowserTexture/Private/WebBrowserTexture.cpp
@@ -0,0 +1,152 @@
+// Copyright 1998-2017 Epic Games, Inc. All Rights Reserved.
+
+#include "WebBrowserTexture.h"
+
+#include "ExternalTexture.h"
+#include "Modules/ModuleManager.h"
+#include "RenderUtils.h"
+#include "RenderingThread.h"
+#include "UObject/WeakObjectPtrTemplates.h"
+#include "WebBrowserTextureResource.h"
+#include "IWebBrowserWindow.h"
+
+/* UWebBrowserTexture structors
+*****************************************************************************/
+
+UWebBrowserTexture::UWebBrowserTexture(const FObjectInitializer& ObjectInitializer)
+	: Super(ObjectInitializer)
+	, AddressX(TA_Clamp)
+	, AddressY(TA_Clamp)
+	, ClearColor(FLinearColor::Black)
+	, Size(0)
+{
+	SampleQueue = MakeShared<FWebBrowserTextureSampleQueue, ESPMode::ThreadSafe>();
+	WebPlayerGuid = FGuid::NewGuid();
+	NeverStream = true;
+}
+
+/* UWebBrowserTexture interface
+*****************************************************************************/
+
+float UWebBrowserTexture::GetAspectRatio() const
+{
+	if (Dimensions.Y == 0)
+	{
+		return 0.0f;
+	}
+
+	return (float)(Dimensions.X) / Dimensions.Y;
+}
+
+int32 UWebBrowserTexture::GetHeight() const
+{
+	return Dimensions.Y;
+}
+
+int32 UWebBrowserTexture::GetWidth() const
+{
+	return Dimensions.X;
+}
+
+/* UTexture interface
+*****************************************************************************/
+
+FTextureResource* UWebBrowserTexture::CreateResource()
+{
+	return new FWebBrowserTextureResource(*this, Dimensions, Size);
+}
+
+EMaterialValueType UWebBrowserTexture::GetMaterialType() const
+{
+	return MCT_TextureExternal;
+}
+
+float UWebBrowserTexture::GetSurfaceWidth() const
+{
+	return Dimensions.X;
+}
+
+float UWebBrowserTexture::GetSurfaceHeight() const
+{
+	return Dimensions.Y;
+}
+
+FGuid UWebBrowserTexture::GetExternalTextureGuid() const
+{
+	return WebPlayerGuid;
+}
+
+void UWebBrowserTexture::SetExternalTextureGuid(FGuid guid)
+{
+	WebPlayerGuid = guid;
+}
+
+/* UObject interface
+*****************************************************************************/
+
+void UWebBrowserTexture::BeginDestroy()
+{
+	UnregisterPlayerGuid();
+
+	Super::BeginDestroy();
+}
+
+FString UWebBrowserTexture::GetDesc()
+{
+	return FString::Printf(TEXT("%ix%i [%s]"), Dimensions.X, Dimensions.Y, GPixelFormats[PF_B8G8R8A8].Name);
+}
+
+
+void UWebBrowserTexture::GetResourceSizeEx(FResourceSizeEx& CumulativeResourceSize)
+{
+	Super::GetResourceSizeEx(CumulativeResourceSize);
+	CumulativeResourceSize.AddUnknownMemoryBytes(Size);
+}
+
+/* UWebBrowserTexture implementation
+*****************************************************************************/
+
+
+void UWebBrowserTexture::TickResource(TSharedPtr<FWebBrowserTextureSample, ESPMode::ThreadSafe> Sample)
+{
+	if (Resource == nullptr)
+	{
+		return;
+	}
+	
+	check(SampleQueue.IsValid());
+
+	if (Sample.IsValid())
+	{
+		SampleQueue.Get()->Enqueue(Sample);
+	}
+
+	// issue a render command to render the current sample
+	FWebBrowserTextureResource::FRenderParams RenderParams;
+	{
+		RenderParams.ClearColor = ClearColor;
+		RenderParams.PlayerGuid = GetExternalTextureGuid();
+		RenderParams.SampleSource = SampleQueue;
+	}
+
+	ENQUEUE_UNIQUE_RENDER_COMMAND_TWOPARAMETER(UWebBrowserTextureResourceRender,
+		FWebBrowserTextureResource*, ResourceParam, (FWebBrowserTextureResource*)Resource,
+		FWebBrowserTextureResource::FRenderParams, RenderParamsParam, RenderParams,
+		{
+			ResourceParam->Render(RenderParamsParam);
+		});
+}
+
+void UWebBrowserTexture::UnregisterPlayerGuid()
+{
+	if (!WebPlayerGuid.IsValid())
+	{
+		return;
+	}
+
+	ENQUEUE_UNIQUE_RENDER_COMMAND_ONEPARAMETER(UWebBrowserTextureUnregisterPlayerGuid,
+		FGuid, PlayerGuid, WebPlayerGuid,
+		{
+			FExternalTextureRegistry::Get().UnregisterExternalTexture(PlayerGuid);
+		});
+}
diff --git a/Engine/Source/Runtime/WebBrowserTexture/Private/WebBrowserTextureModule.cpp b/Engine/Source/Runtime/WebBrowserTexture/Private/WebBrowserTextureModule.cpp
new file mode 100644
--- /dev/null
+++ b/Engine/Source/Runtime/WebBrowserTexture/Private/WebBrowserTextureModule.cpp
@@ -0,0 +1,22 @@
+// Copyright 1998-2017 Epic Games, Inc. All Rights Reserved.
+
+#include "WebBrowserTextureModule.h"
+
+
+class FWebBrowserTextureModule : public IWebBrowserTextureModule
+{
+private:
+	// IModuleInterface Interface
+	virtual void StartupModule() override;
+	virtual void ShutdownModule() override;
+};
+
+IMPLEMENT_MODULE( FWebBrowserTextureModule, WebBrowserTexture );
+
+void FWebBrowserTextureModule::StartupModule()
+{
+}
+
+void FWebBrowserTextureModule::ShutdownModule()
+{
+}
diff --git a/Engine/Source/Runtime/WebBrowserTexture/Private/WebBrowserTextureResource.cpp b/Engine/Source/Runtime/WebBrowserTexture/Private/WebBrowserTextureResource.cpp
new file mode 100644
--- /dev/null
+++ b/Engine/Source/Runtime/WebBrowserTexture/Private/WebBrowserTextureResource.cpp
@@ -0,0 +1,296 @@
+// Copyright 1998-2017 Epic Games, Inc. All Rights Reserved.
+
+#include "WebBrowserTextureResource.h"
+
+#include "DeviceProfiles/DeviceProfile.h"
+#include "DeviceProfiles/DeviceProfileManager.h"
+#include "ExternalTexture.h"
+#include "PipelineStateCache.h"
+#include "SceneUtils.h"
+#include "Shader.h"
+#include "StaticBoundShaderState.h"
+#include "RenderUtils.h"
+#include "RHIStaticStates.h"
+#include "ExternalTexture.h"
+#include "WebBrowserTexture.h"
+
+
+#define WebBrowserTextureRESOURCE_TRACE_RENDER 0
+
+
+/* FWebBrowserTextureResource structors
+ *****************************************************************************/
+
+FWebBrowserTextureResource::FWebBrowserTextureResource(UWebBrowserTexture& InOwner, FIntPoint& InOwnerDim, SIZE_T& InOwnerSize)
+	: Cleared(false)
+	, CurrentClearColor(FLinearColor::Transparent)
+	, Owner(InOwner)
+	, OwnerDim(InOwnerDim)
+	, OwnerSize(InOwnerSize)
+{ 
+	FPlatformMisc::LowLevelOutputDebugStringf(TEXT("FWebBrowserTextureResource:FWebBrowserTextureResource %d %d"), OwnerDim.X, OwnerDim.Y);
+
+}
+
+
+/* FWebBrowserTextureResource interface
+ *****************************************************************************/
+
+void FWebBrowserTextureResource::Render(const FRenderParams& Params)
+{
+	check(IsInRenderingThread());
+
+	TSharedPtr<FWebBrowserTextureSampleQueue, ESPMode::ThreadSafe> SampleSource = Params.SampleSource.Pin();
+
+	if (SampleSource.IsValid())
+	{
+		// get the most current sample to be rendered
+		TSharedPtr<FWebBrowserTextureSample, ESPMode::ThreadSafe> Sample;
+		bool SampleValid = false;
+		
+		while (SampleSource->Peek(Sample) && Sample.IsValid())
+		{
+			SampleValid = SampleSource->Dequeue(Sample);
+		}
+
+		if (!SampleValid)
+		{
+			return; // no sample to render
+		}
+
+		check(Sample.IsValid());
+
+		// render the sample
+		CopySample(Sample, Params.ClearColor);
+
+		if (!GSupportsImageExternal && Params.PlayerGuid.IsValid())
+		{
+			FTextureRHIRef VideoTexture = (FTextureRHIRef)Owner.TextureReference.TextureReferenceRHI;
+			FExternalTextureRegistry::Get().RegisterExternalTexture(Params.PlayerGuid, VideoTexture, SamplerStateRHI, Sample->GetScaleRotation(), Sample->GetOffset());
+		}
+	}
+	else if (!Cleared)
+	{
+		ClearTexture(Params.ClearColor);
+
+		if (!GSupportsImageExternal && Params.PlayerGuid.IsValid())
+		{
+			FTextureRHIRef VideoTexture = (FTextureRHIRef)Owner.TextureReference.TextureReferenceRHI;
+			FExternalTextureRegistry::Get().RegisterExternalTexture(Params.PlayerGuid, VideoTexture, SamplerStateRHI, FLinearColor(1.0f, 0.0f, 0.0f, 1.0f), FLinearColor(0.0f, 0.0f, 0.0f, 0.0f));
+		}
+	}
+}
+
+
+/* FRenderTarget interface
+ *****************************************************************************/
+
+FIntPoint FWebBrowserTextureResource::GetSizeXY() const
+{
+	return FIntPoint(Owner.GetWidth(), Owner.GetHeight());
+}
+
+
+/* FTextureResource interface
+ *****************************************************************************/
+
+FString FWebBrowserTextureResource::GetFriendlyName() const
+{
+	return Owner.GetPathName();
+}
+
+
+uint32 FWebBrowserTextureResource::GetSizeX() const
+{
+	return Owner.GetWidth();
+}
+
+
+uint32 FWebBrowserTextureResource::GetSizeY() const
+{
+	return Owner.GetHeight();
+}
+
+
+void FWebBrowserTextureResource::InitDynamicRHI()
+{
+	// create the sampler state
+	FSamplerStateInitializerRHI SamplerStateInitializer(
+		(ESamplerFilter)UDeviceProfileManager::Get().GetActiveProfile()->GetTextureLODSettings()->GetSamplerFilter(&Owner),
+		(Owner.AddressX == TA_Wrap) ? AM_Wrap : ((Owner.AddressX == TA_Clamp) ? AM_Clamp : AM_Mirror),
+		(Owner.AddressY == TA_Wrap) ? AM_Wrap : ((Owner.AddressY == TA_Clamp) ? AM_Clamp : AM_Mirror),
+		AM_Wrap
+	);
+
+	SamplerStateRHI = RHICreateSamplerState(SamplerStateInitializer);
+}
+
+
+void FWebBrowserTextureResource::ReleaseDynamicRHI()
+{
+	Cleared = false;
+
+	InputTarget.SafeRelease();
+	OutputTarget.SafeRelease();
+	RenderTargetTextureRHI.SafeRelease();
+	TextureRHI.SafeRelease();
+
+	UpdateTextureReference(nullptr);
+}
+
+
+/* FWebBrowserTextureResource implementation
+ *****************************************************************************/
+
+void FWebBrowserTextureResource::ClearTexture(const FLinearColor& ClearColor)
+{
+	FPlatformMisc::LowLevelOutputDebugStringf(TEXT("FWebBrowserTextureResource:ClearTexture"));
+	// create output render target if we don't have one yet
+	const uint32 OutputCreateFlags = TexCreate_Dynamic | TexCreate_SRGB;
+
+	if ((ClearColor != CurrentClearColor) || !OutputTarget.IsValid() || ((OutputTarget->GetFlags() & OutputCreateFlags) != OutputCreateFlags))
+	{
+		FRHIResourceCreateInfo CreateInfo = {
+			FClearValueBinding(ClearColor)
+		};
+
+		TRefCountPtr<FRHITexture2D> DummyTexture2DRHI;
+
+		RHICreateTargetableShaderResource2D(
+			2,
+			2,
+			PF_B8G8R8A8,
+			1,
+			OutputCreateFlags,
+			TexCreate_RenderTargetable,
+			false,
+			CreateInfo,
+			OutputTarget,
+			DummyTexture2DRHI
+		);
+
+		CurrentClearColor = ClearColor;
+		UpdateResourceSize();
+	}
+
+	if (RenderTargetTextureRHI != OutputTarget)
+	{
+		UpdateTextureReference(OutputTarget);
+	}
+
+	// draw the clear color
+	FRHICommandListImmediate& CommandList = FRHICommandListExecutor::GetImmediateCommandList();
+	{
+		FRHIRenderTargetView View = FRHIRenderTargetView(RenderTargetTextureRHI, ERenderTargetLoadAction::EClear);
+		FRHISetRenderTargetsInfo Info(1, &View, FRHIDepthRenderTargetView());
+		CommandList.SetRenderTargetsAndClear(Info);
+		CommandList.SetViewport(0, 0, 0.0f, RenderTargetTextureRHI->GetSizeX(), RenderTargetTextureRHI->GetSizeY(), 1.0f);
+		CommandList.TransitionResource(EResourceTransitionAccess::EReadable, RenderTargetTextureRHI);
+	}
+
+	Cleared = true;
+}
+
+void FWebBrowserTextureResource::CopySample(const TSharedPtr<FWebBrowserTextureSample, ESPMode::ThreadSafe>& Sample, const FLinearColor& ClearColor)
+{
+	FRHITexture* SampleTexture = Sample->GetTexture();
+	FRHITexture2D* SampleTexture2D = (SampleTexture != nullptr) ? SampleTexture->GetTexture2D() : nullptr;
+	// If the sample already provides a texture resource, we simply use that
+	// as the output render target. If the sample only provides raw data, then
+	// we create our own output render target and copy the data into it.
+	if (SampleTexture2D != nullptr)
+	{
+	FPlatformMisc::LowLevelOutputDebugStringf(TEXT("FWebBrowserTextureResource:CopySample 1"));
+		// use sample's texture as the new render target.
+		if (TextureRHI != SampleTexture2D)
+		{
+			FPlatformMisc::LowLevelOutputDebugStringf(TEXT("FWebBrowserTextureResource:CopySample 11"));
+			UpdateTextureReference(SampleTexture2D);
+
+			OutputTarget.SafeRelease();
+			UpdateResourceSize();
+		}
+	}
+	else
+	{
+		FPlatformMisc::LowLevelOutputDebugStringf(TEXT("FWebBrowserTextureResource:CopySample 2"));
+		// create a new output render target if necessary
+		const uint32 OutputCreateFlags = TexCreate_Dynamic | TexCreate_SRGB;
+		const FIntPoint SampleDim = Sample->GetDim();
+
+		if ((ClearColor != CurrentClearColor) || !OutputTarget.IsValid() || (OutputTarget->GetSizeXY() != SampleDim) || ((OutputTarget->GetFlags() & OutputCreateFlags) != OutputCreateFlags))
+		{
+			FPlatformMisc::LowLevelOutputDebugStringf(TEXT("FWebBrowserTextureResource:CopySample 1"));
+			TRefCountPtr<FRHITexture2D> DummyTexture2DRHI;
+
+			FRHIResourceCreateInfo CreateInfo = {
+				FClearValueBinding(ClearColor)
+			};
+
+			RHICreateTargetableShaderResource2D(
+				SampleDim.X,
+				SampleDim.Y,
+				PF_B8G8R8A8,
+				1,
+				OutputCreateFlags,
+				TexCreate_RenderTargetable,
+				false,
+				CreateInfo,
+				OutputTarget,
+				DummyTexture2DRHI
+			);
+
+			CurrentClearColor = ClearColor;
+			UpdateResourceSize();
+		}
+
+		if (RenderTargetTextureRHI != OutputTarget)
+		{
+			UpdateTextureReference(OutputTarget);
+		}
+
+		// copy sample data to output render target
+		FUpdateTextureRegion2D Region(0, 0, 0, 0, SampleDim.X, SampleDim.Y);
+		RHIUpdateTexture2D(RenderTargetTextureRHI.GetReference(), 0, Region, Sample->GetStride(), (uint8*)Sample->GetBuffer());
+	}
+	Cleared = false;
+}
+
+
+void FWebBrowserTextureResource::UpdateResourceSize()
+{
+	FPlatformMisc::LowLevelOutputDebugStringf(TEXT("FWebBrowserTextureResource:UpdateResourceSize"));
+
+	SIZE_T ResourceSize = 0;
+
+	if (InputTarget.IsValid())
+	{
+		ResourceSize += CalcTextureSize(InputTarget->GetSizeX(), InputTarget->GetSizeY(), InputTarget->GetFormat(), 1);
+	}
+
+	if (OutputTarget.IsValid())
+	{
+		ResourceSize += CalcTextureSize(OutputTarget->GetSizeX(), OutputTarget->GetSizeY(), OutputTarget->GetFormat(), 1);
+	}
+
+	OwnerSize = ResourceSize;
+}
+
+
+void FWebBrowserTextureResource::UpdateTextureReference(FRHITexture2D* NewTexture)
+{
+	FPlatformMisc::LowLevelOutputDebugStringf(TEXT("FWebBrowserTextureResource:UpdateTextureReference"));
+	TextureRHI = NewTexture;
+	RenderTargetTextureRHI = NewTexture;
+
+	RHIUpdateTextureReference(Owner.TextureReference.TextureReferenceRHI, NewTexture);
+
+	if (RenderTargetTextureRHI != nullptr)
+	{
+		OwnerDim = FIntPoint(RenderTargetTextureRHI->GetSizeX(), RenderTargetTextureRHI->GetSizeY());
+	}
+	else
+	{
+		OwnerDim = FIntPoint::ZeroValue;
+	}
+}
diff --git a/Engine/Source/Runtime/WebBrowserTexture/Public/WebBrowserTexture.h b/Engine/Source/Runtime/WebBrowserTexture/Public/WebBrowserTexture.h
new file mode 100644
--- /dev/null
+++ b/Engine/Source/Runtime/WebBrowserTexture/Public/WebBrowserTexture.h
@@ -0,0 +1,98 @@
+// Copyright 1998-2017 Epic Games, Inc. All Rights Reserved.
+
+#pragma once
+
+#include "CoreTypes.h"
+#include "Containers/EnumAsByte.h"
+#include "Engine/Texture.h"
+#include "Math/Color.h"
+#include "Math/IntPoint.h"
+#include "Misc/Timespan.h"
+#include "Templates/SharedPointer.h"
+#include "UObject/ObjectMacros.h"
+#include "UObject/ScriptMacros.h"
+#include "WebBrowserTextureSample.h"
+
+#include "WebBrowserTexture.generated.h"
+
+/**
+* Implements a texture asset for rendering webbrowser output for Android.
+*/
+UCLASS(hidecategories = (Adjustments, Compositing, LevelOfDetail, Object))
+class WEBBROWSERTEXTURE_API UWebBrowserTexture
+	: public UTexture
+{
+	GENERATED_UCLASS_BODY()
+
+	/** The addressing mode to use for the X axis. */
+	TEnumAsByte<TextureAddress> AddressX;
+
+	/** The addressing mode to use for the Y axis. */
+	TEnumAsByte<TextureAddress> AddressY;
+
+	/** Whether to clear the texture when no media is being played (default = enabled). */
+	bool AutoClear;
+
+	/** The color used to clear the texture if AutoClear is enabled (default = black). */
+	FLinearColor ClearColor;
+
+public:
+
+	/**
+	* Gets the current aspect ratio of the texture.
+	*
+	* @return Texture aspect ratio.
+	* @see GetHeight, GetWidth
+	*/
+	float GetAspectRatio() const;
+
+	/**
+	* Gets the current height of the texture.
+	*
+	* @return Texture height (in pixels).
+	* @see GetAspectRatio, GetWidth
+	*/
+	int32 GetHeight() const;
+
+	/**
+	* Gets the current width of the texture.
+	*
+	* @return Texture width (in pixels).
+	* @see GetAspectRatio, GetHeight
+	*/
+	int32 GetWidth() const;
+
+public:
+	//~ UTexture interface.
+	virtual void BeginDestroy() override;
+	virtual FTextureResource* CreateResource() override;
+	virtual EMaterialValueType GetMaterialType() const override;
+	virtual float GetSurfaceWidth() const override;
+	virtual float GetSurfaceHeight() const override;
+	virtual FGuid GetExternalTextureGuid() const override;
+
+public:
+	//~ UObject interface.
+	virtual FString GetDesc() override;
+	virtual void GetResourceSizeEx(FResourceSizeEx& CumulativeResourceSize) override;
+	
+	void TickResource(TSharedPtr<FWebBrowserTextureSample, ESPMode::ThreadSafe> Sample);
+
+	void SetExternalTextureGuid(FGuid guid);
+protected:
+	/** Unregister the player's external texture GUID. */
+	void UnregisterPlayerGuid();
+
+private:
+	/** Current width and height of the resource (in pixels). */
+	FIntPoint Dimensions;
+
+	/** The previously used player GUID. */
+	FGuid WebPlayerGuid;
+
+	/** Texture sample queue. */
+	TSharedPtr<FWebBrowserTextureSampleQueue, ESPMode::ThreadSafe> SampleQueue;
+
+	/** Current size of the resource (in bytes).*/
+	SIZE_T Size;
+};
diff --git a/Engine/Source/Runtime/WebBrowserTexture/Public/WebBrowserTextureModule.h b/Engine/Source/Runtime/WebBrowserTexture/Public/WebBrowserTextureModule.h
new file mode 100644
--- /dev/null
+++ b/Engine/Source/Runtime/WebBrowserTexture/Public/WebBrowserTextureModule.h
@@ -0,0 +1,14 @@
+// Copyright 1998-2017 Epic Games, Inc. All Rights Reserved.
+
+#pragma once
+
+#include "CoreMinimal.h"
+#include "Modules/ModuleInterface.h"
+#include "Modules/ModuleManager.h"
+
+/**
+ * WebBrowserTextureModule interface
+ */
+class IWebBrowserTextureModule : public IModuleInterface
+{
+};
diff --git a/Engine/Source/Runtime/WebBrowserTexture/Public/WebBrowserTextureResource.h b/Engine/Source/Runtime/WebBrowserTexture/Public/WebBrowserTextureResource.h
new file mode 100644
--- /dev/null
+++ b/Engine/Source/Runtime/WebBrowserTexture/Public/WebBrowserTextureResource.h
@@ -0,0 +1,131 @@
+// Copyright 1998-2017 Epic Games, Inc. All Rights Reserved.
+
+#pragma once
+
+#include "CoreTypes.h"
+#include "Containers/Queue.h"
+#include "Math/Color.h"
+#include "Misc/Guid.h"
+#include "Misc/Timespan.h"
+#include "Templates/RefCounting.h"
+#include "Templates/SharedPointer.h"
+#include "TextureResource.h"
+#include "UnrealClient.h"
+#include "WebBrowserTextureSample.h"
+
+class UWebBrowserTexture;
+
+/**
+ * Texture resource type for webbrowser textures.
+ */
+class FWebBrowserTextureResource
+	: public FRenderTarget
+	, public FTextureResource
+{
+public:
+
+	/** 
+	 * Creates and initializes a new instance.
+	 *
+	 * @param InOwner The Movie texture object to create a resource for (must not be nullptr).
+	 * @param InOwnerDim The initial width and height of the texture.
+	 * @param InOwnerSize The initial memory allocated to the texture.
+	 */
+	FWebBrowserTextureResource(UWebBrowserTexture& InOwner, FIntPoint& InOwnerDim, SIZE_T& InOwnerSize);
+
+	/** Virtual destructor. */
+	virtual ~FWebBrowserTextureResource() { }
+
+public:
+
+	/** Parameters for the Render method. */
+	struct FRenderParams
+	{
+		/** The clear color to use when clearing the texture. */
+		FLinearColor ClearColor;
+
+		/** Guid associated with the texture. */
+		FGuid PlayerGuid;
+
+		/** The sample source to render. */
+		TWeakPtr<FWebBrowserTextureSampleQueue, ESPMode::ThreadSafe> SampleSource;
+
+	};
+
+	/**
+	 * Render the texture resource.
+	 *
+	 * This method is called on the render thread by the WebBrowserTexture that owns this
+	 * texture resource to clear or redraw the resource using the given parameters.
+	 *
+	 * @param Params Render parameters.
+	 */
+	void Render(const FRenderParams& Params);
+
+public:
+
+	//~ FRenderTarget interface
+
+	virtual FIntPoint GetSizeXY() const override;
+
+public:
+
+	//~ FTextureResource interface
+
+	virtual FString GetFriendlyName() const override;
+	virtual uint32 GetSizeX() const override;
+	virtual uint32 GetSizeY() const override;
+	virtual void InitDynamicRHI() override;
+	virtual void ReleaseDynamicRHI() override;
+
+protected:
+
+	/**
+	 * Clear the texture using the given clear color.
+	 *
+	 * @param ClearColor The clear color to use.
+	 */
+	void ClearTexture(const FLinearColor& ClearColor);
+
+	/**
+	 * Render the given texture sample by using it as or copying it to the render target.
+	 *
+	 * @param Sample The texture sample to copy.
+	 * @param ClearColor The clear color to use for the output texture.
+	 * @see ConvertSample
+	 */
+	void CopySample(const TSharedPtr<FWebBrowserTextureSample, ESPMode::ThreadSafe>& Sample, const FLinearColor& ClearColor);
+
+	/** Calculates the current resource size and notifies the owner texture. */
+	void UpdateResourceSize();
+
+	/**
+	 * Set the owner's texture reference to the given texture.
+	 *
+	 * @param NewTexture The texture to set.
+	 */
+	void UpdateTextureReference(FRHITexture2D* NewTexture);
+
+private:
+
+	/** Whether the texture has been cleared. */
+	bool Cleared;
+
+	/** Tracks the current clear color. */
+	FLinearColor CurrentClearColor;
+
+	/** Input render target if the texture samples don't provide one (for conversions). */
+	TRefCountPtr<FRHITexture2D> InputTarget;
+
+	/** Output render target if the texture samples don't provide one. */
+	TRefCountPtr<FRHITexture2D> OutputTarget;
+
+	/** The webbrowser texture that owns this resource. */
+	UWebBrowserTexture& Owner;
+
+	/** Reference to the owner's texture dimensions field. */
+	FIntPoint& OwnerDim;
+
+	/** Reference to the owner's texture size field. */
+	SIZE_T& OwnerSize;
+};
diff --git a/Engine/Source/Runtime/WebBrowserTexture/Public/WebBrowserTextureSample.h b/Engine/Source/Runtime/WebBrowserTexture/Public/WebBrowserTextureSample.h
new file mode 100644
--- /dev/null
+++ b/Engine/Source/Runtime/WebBrowserTexture/Public/WebBrowserTextureSample.h
@@ -0,0 +1,392 @@
+// Copyright 1998-2016 Epic Games, Inc. All Rights Reserved.
+
+#pragma once
+
+#include "CoreTypes.h"
+#include "CoreGlobals.h"
+#include "Containers/Queue.h"
+#include "Math/IntPoint.h"
+#include "Math/Range.h"
+#include "MediaObjectPool.h"
+#include "RHI.h"
+#include "RHIUtilities.h"
+#include "Templates/SharedPointer.h"
+
+/**
+* Texture sample generated by the WebBrowser.
+*/
+class FWebBrowserTextureSample
+	: public IMediaPoolable
+{
+public:
+
+	/** Default constructor. */
+	FWebBrowserTextureSample()
+		: Buffer(nullptr)
+		, BufferSize(0)
+		, Dim(FIntPoint::ZeroValue)
+		, ScaleRotation(FLinearColor(1.0f, 0.0f, 0.0f, 1.0f))
+		, Offset(FLinearColor(0.0f, 0.0f, 0.0f, 0.0f))
+	{ }
+
+	/** Virtual destructor. */
+	virtual ~FWebBrowserTextureSample()
+	{
+		if (BufferSize > 0)
+		{
+			FMemory::Free(Buffer);
+		}
+	}
+
+public:
+
+	/**
+	* Get a writable pointer to the sample buffer.
+	*
+	* @return Sample buffer.
+	*/
+	void* GetMutableBuffer()
+	{
+		return Buffer;
+	}
+
+	/**
+	* Initialize the sample.
+	*
+	* @param InDim The sample buffer's width and height (in pixels).
+	* @param InDuration The duration for which the sample is valid.
+	* @return true on success, false otherwise.
+	*/
+	bool Initialize(const FIntPoint& InDim)
+	{
+		BufferSize = 0;
+
+		if (InDim.GetMin() <= 0)
+		{ 
+			return false;
+		}
+
+		Dim = InDim;
+
+		return true;
+	}
+
+	/**
+	* Initialize the sample with a memory buffer.
+	*
+	* @param InBuffer The buffer containing the sample data.
+	* @param Copy Whether the buffer should be copied (true) or referenced (false).
+	* @see InitializeTexture
+	*/
+	void InitializeBuffer(void* InBuffer, bool Copy)
+	{
+		if (Copy)
+		{
+			SIZE_T RequiredBufferSize = Dim.X * Dim.Y * sizeof(int32);
+			if (BufferSize < RequiredBufferSize)
+			{
+				if (BufferSize == 0)
+				{
+					Buffer = FMemory::Malloc(RequiredBufferSize);
+				}
+				else
+				{
+					Buffer = FMemory::Realloc(Buffer, RequiredBufferSize);
+				}
+
+				BufferSize = RequiredBufferSize;
+			}
+
+			FMemory::Memcpy(Buffer, InBuffer, RequiredBufferSize);
+		}
+		else
+		{
+			if (BufferSize > 0)
+			{
+				FMemory::Free(Buffer);
+				BufferSize = 0;
+			}
+
+			Buffer = InBuffer;
+		}
+	}
+
+	/**
+	* Initialize the sample with a texture resource.
+	*
+	* @return The texture resource object that will hold the sample data.
+	* @note This method must be called on the render thread.
+	* @see InitializeBuffer
+	*/
+	FRHITexture2D* InitializeTexture()
+	{
+		check(IsInRenderingThread());
+
+		if (Texture.IsValid() && (Texture->GetSizeXY() == Dim))
+		{
+			return Texture;
+		}
+
+		const uint32 CreateFlags = TexCreate_Dynamic | TexCreate_SRGB;
+
+		TRefCountPtr<FRHITexture2D> DummyTexture2DRHI;
+		FRHIResourceCreateInfo CreateInfo;
+
+		RHICreateTargetableShaderResource2D(
+			Dim.X,
+			Dim.Y,
+			PF_B8G8R8A8,
+			1,
+			CreateFlags,
+			TexCreate_RenderTargetable,
+			false,
+			CreateInfo,
+			Texture,
+			DummyTexture2DRHI
+		);
+
+		return Texture;
+	}
+
+	/**
+	* Set the sample Scale, Rotation, Offset.
+	*
+	* @param InScaleRotation The sample scale and rotation transform (2x2).
+	* @param InOffset The sample offset.
+	*/
+	void SetScaleRotationOffset(FVector4& InScaleRotation, FVector4& InOffset)
+	{
+		ScaleRotation = FLinearColor(InScaleRotation.X, InScaleRotation.Y, InScaleRotation.Z, InScaleRotation.W);
+		Offset = FLinearColor(InOffset.X, InOffset.Y, InOffset.Z, InOffset.W);
+	}
+
+public:
+
+	//~ IMediaTextureSample interface
+
+	virtual const void* GetBuffer() 
+	{
+		return Buffer;
+	}
+
+	virtual FIntPoint GetDim() const
+	{
+		return Dim;
+	}
+
+	virtual uint32 GetStride() const
+	{
+		return Dim.X * sizeof(int32);
+	}
+
+#if WITH_ENGINE
+
+	virtual FRHITexture* GetTexture() const
+	{
+		return Texture.GetReference();
+	}
+
+#endif //WITH_ENGINE
+
+	virtual FLinearColor GetScaleRotation() const
+	{
+		return ScaleRotation;
+	}
+
+	virtual FLinearColor GetOffset() const
+	{
+		return Offset;
+	}
+
+private:
+
+	/** The sample's data buffer. */
+	void* Buffer;
+
+	/** Current allocation size of Buffer. */
+	SIZE_T BufferSize;
+
+	/** Width and height of the texture sample. */
+	FIntPoint Dim;
+
+	/** ScaleRotation for the sample. */
+	FLinearColor ScaleRotation;
+
+	/** Offset for the sample. */
+	FLinearColor Offset;
+
+#if WITH_ENGINE
+
+	/** Texture resource. */
+	TRefCountPtr<FRHITexture2D> Texture;
+
+#endif //WITH_ENGINE
+};
+
+
+class FWebBrowserTextureSampleQueue
+{
+public:
+
+	/** Default constructor. */
+	FWebBrowserTextureSampleQueue()
+		: NumSamples(0)
+		, PendingFlushes(0)
+	{ }
+
+	/** Virtual destructor. */
+	virtual ~FWebBrowserTextureSampleQueue() { }
+
+public:
+
+	/**
+	* Get the number of samples in the queue.
+	*
+	* @return Number of samples.
+	* @see Enqueue, Dequeue, Peek
+	*/
+	int32 Num() const
+	{
+		return NumSamples;
+	}
+
+public:
+
+	//~ TMediaSampleSource interface (to be called only from consumer thread)
+
+	virtual bool Dequeue(TSharedPtr<FWebBrowserTextureSample, ESPMode::ThreadSafe>& OutSample)
+	{
+		DoPendingFlushes();
+
+		TSharedPtr<FWebBrowserTextureSample, ESPMode::ThreadSafe> Sample;
+
+		if (!Samples.Peek(Sample))
+		{
+			return false; // empty queue
+		}
+
+		if (!Sample.IsValid())
+		{
+			return false; // pending flush
+		}
+
+		Samples.Pop();
+
+		FPlatformAtomics::InterlockedDecrement(&NumSamples);
+		check(NumSamples >= 0);
+
+		OutSample = Sample;
+
+		return true;
+	}
+
+	virtual bool Peek(TSharedPtr<FWebBrowserTextureSample, ESPMode::ThreadSafe>& OutSample)
+	{
+		DoPendingFlushes();
+
+		TSharedPtr<FWebBrowserTextureSample, ESPMode::ThreadSafe> Sample;
+
+		if (!Samples.Peek(Sample))
+		{
+			return false; // empty queue
+		}
+
+		if (!Sample.IsValid())
+		{
+			return false; // pending flush
+		}
+
+		OutSample = Sample;
+
+		return true;
+	}
+
+	virtual bool Pop()
+	{
+		TSharedPtr<FWebBrowserTextureSample, ESPMode::ThreadSafe> Sample;
+
+		if (!Samples.Peek(Sample))
+		{
+			return false; // empty queue
+		}
+
+		if (!Sample.IsValid())
+		{
+			return false; // pending flush
+		}
+
+		Samples.Pop();
+
+		FPlatformAtomics::InterlockedDecrement(&NumSamples);
+		check(NumSamples >= 0);
+
+		return true;
+	}
+
+public:
+
+	//~ TMediaSampleSink interface (to be called only from producer thread)
+
+	virtual bool Enqueue(const TSharedPtr<FWebBrowserTextureSample, ESPMode::ThreadSafe>& Sample)
+	{
+		if (Sample.IsValid())
+		{
+			FPlatformAtomics::InterlockedIncrement(&NumSamples);
+		}
+
+		if (!Samples.Enqueue(Sample))
+		{
+			if (Sample.IsValid())
+			{
+				FPlatformAtomics::InterlockedDecrement(&NumSamples);
+			}
+
+			return false;
+		}
+
+		return true;
+	}
+
+	virtual void RequestFlush()
+	{
+		Samples.Enqueue(nullptr); // insert flush marker
+		FPlatformAtomics::InterlockedIncrement(&PendingFlushes);
+	}
+
+protected:
+
+	/** Perform any pending flushes. */
+	void DoPendingFlushes()
+	{
+		TSharedPtr<FWebBrowserTextureSample, ESPMode::ThreadSafe> Sample;
+
+		while ((PendingFlushes > 0) && Samples.Dequeue(Sample))
+		{
+			if (Sample.IsValid())
+			{
+				FPlatformAtomics::InterlockedDecrement(&NumSamples);
+				check(NumSamples >= 0);
+			}
+			else
+			{
+				FPlatformAtomics::InterlockedDecrement(&PendingFlushes);
+			}
+		}
+	}
+
+private:
+
+	/** Number of samples in the queue. */
+	int32 NumSamples;
+
+	/** Number of pending flushes. */
+	int32 PendingFlushes;
+
+	/** Audio sample queue. */
+	TQueue<TSharedPtr<FWebBrowserTextureSample, ESPMode::ThreadSafe>, EQueueMode::Mpsc> Samples;
+};
+
+
+/** Implements a pool for WebBrowser's texture sample objects. */
+class FWebBrowserTextureSamplePool : public TMediaObjectPool<FWebBrowserTextureSample> { };
+
diff --git a/Engine/Source/Runtime/WebBrowserTexture/WebBrowserTexture.Build.cs b/Engine/Source/Runtime/WebBrowserTexture/WebBrowserTexture.Build.cs
new file mode 100644
--- /dev/null
+++ b/Engine/Source/Runtime/WebBrowserTexture/WebBrowserTexture.Build.cs
@@ -0,0 +1,47 @@
+// Copyright 1998-2017 Epic Games, Inc. All Rights Reserved.
+
+using UnrealBuildTool;
+using System.IO;
+
+public class WebBrowserTexture : ModuleRules
+{
+	public WebBrowserTexture(ReadOnlyTargetRules Target) : base(Target)
+	{
+        if (Target.Platform == UnrealTargetPlatform.Android)
+		{
+			PublicIncludePaths.Add("Runtime/WebBrowserTexture/Public");
+			PrivateIncludePaths.Add("Runtime/WebBrowserTexture/Private");
+			
+			// Needed for external texture support
+			PublicIncludePaths.AddRange(
+				new string[] {
+				"Runtime/MediaUtils/Public",
+				});
+
+			PrivateIncludePaths.AddRange(
+				new string[] {
+				"Runtime/MediaUtils/Private",
+				});
+
+			PrivateDependencyModuleNames.AddRange(
+				new string[]
+				{
+					"Core",
+					"CoreUObject",
+					"ApplicationCore",
+					"RHI",
+					"InputCore",
+					"Slate",
+					"SlateCore",
+					"Serialization",
+					"MediaUtils",
+					"RenderCore",
+					"Engine",
+					"ShaderCore",
+					"UtilityShaders",
+					"WebBrowser"
+				}
+			);
+		}
+	}
+}
