Fix PCH dependencies in WebBrowserTexture module
#ue4

diff --git a/Engine/Source/Runtime/WebBrowserTexture/WebBrowserTexture.Build.cs b/Engine/Source/Runtime/WebBrowserTexture/WebBrowserTexture.Build.cs
--- a/Engine/Source/Runtime/WebBrowserTexture/WebBrowserTexture.Build.cs
+++ b/Engine/Source/Runtime/WebBrowserTexture/WebBrowserTexture.Build.cs
@@ -7,6 +7,14 @@ public class WebBrowserTexture : ModuleRules
 {
 	public WebBrowserTexture(ReadOnlyTargetRules Target) : base(Target)
 	{
+		PrivateDependencyModuleNames.AddRange(
+			new string[]
+			{
+				"Core",
+				"CoreUObject",
+			}
+		);
+		
         if (Target.Platform == UnrealTargetPlatform.Android)
 		{
 			PublicIncludePaths.Add("Runtime/WebBrowserTexture/Public");
@@ -26,8 +34,6 @@ public class WebBrowserTexture : ModuleRules
 			PrivateDependencyModuleNames.AddRange(
 				new string[]
 				{
-					"Core",
-					"CoreUObject",
 					"ApplicationCore",
 					"RHI",
 					"InputCore",
