Fix more WebBrowser dependencies with editor
#ue4

diff --git a/Engine/Plugins/Runtime/WebBrowserWidget/Source/WebBrowserWidget/Private/WebBrowser.cpp b/Engine/Plugins/Runtime/WebBrowserWidget/Source/WebBrowserWidget/Private/WebBrowser.cpp
--- a/Engine/Plugins/Runtime/WebBrowserWidget/Source/WebBrowserWidget/Private/WebBrowser.cpp
+++ b/Engine/Plugins/Runtime/WebBrowserWidget/Source/WebBrowserWidget/Private/WebBrowser.cpp
@@ -16,7 +16,6 @@
 #include "Materials/Material.h"
 #include "Factories/MaterialFactoryNew.h"
 #include "AssetRegistryModule.h"
-#include "WebBrowserTexture.h"
 #include "PackageHelperFunctions.h"
 #endif
 
diff --git a/Engine/Source/Runtime/WebBrowserTexture/WebBrowserTexture.Build.cs b/Engine/Source/Runtime/WebBrowserTexture/WebBrowserTexture.Build.cs
--- a/Engine/Source/Runtime/WebBrowserTexture/WebBrowserTexture.Build.cs
+++ b/Engine/Source/Runtime/WebBrowserTexture/WebBrowserTexture.Build.cs
@@ -7,14 +7,6 @@ public class WebBrowserTexture : ModuleRules
 {
 	public WebBrowserTexture(ReadOnlyTargetRules Target) : base(Target)
 	{
-		PrivateDependencyModuleNames.AddRange(
-			new string[]
-			{
-				"Core",
-				"CoreUObject",
-			}
-		);
-		
         if (Target.Platform == UnrealTargetPlatform.Android)
 		{
 			PublicIncludePaths.Add("Runtime/WebBrowserTexture/Public");
@@ -34,6 +26,8 @@ public class WebBrowserTexture : ModuleRules
 			PrivateDependencyModuleNames.AddRange(
 				new string[]
 				{
+					"Core",
+					"CoreUObject",
 					"ApplicationCore",
 					"RHI",
 					"InputCore",
@@ -49,5 +43,9 @@ public class WebBrowserTexture : ModuleRules
 				}
 			);
 		}
+		else
+		{
+			PrecompileForTargets = ModuleRules.PrecompileTargetsType.None;
+		}
 	}
 }
