Fix CIS

diff --git a/Engine/Source/Runtime/WebBrowserTexture/WebBrowserTexture.Build.cs b/Engine/Source/Runtime/WebBrowserTexture/WebBrowserTexture.Build.cs
--- a/Engine/Source/Runtime/WebBrowserTexture/WebBrowserTexture.Build.cs
+++ b/Engine/Source/Runtime/WebBrowserTexture/WebBrowserTexture.Build.cs
@@ -10,7 +10,8 @@ public class WebBrowserTexture : ModuleRules
         // WebBrowserTexture objects are needed only on Android, but we also need to be able to
         // cook the asset so we must include it in editor builds
         if (Target.Platform == UnrealTargetPlatform.Android ||
-            Target.Platform == UnrealTargetPlatform.IOS || 
+            Target.Platform == UnrealTargetPlatform.IOS ||
+            Target.Platform == UnrealTargetPlatform.TVOS ||
             Target.bBuildEditor == true)
 		{
 			PublicIncludePaths.Add("Runtime/WebBrowserTexture/Public");
