WebBrowser texture RenderPasses update

diff --git a/Engine/Source/Runtime/WebBrowserTexture/Private/WebBrowserTextureResource.cpp b/Engine/Source/Runtime/WebBrowserTexture/Private/WebBrowserTextureResource.cpp
--- a/Engine/Source/Runtime/WebBrowserTexture/Private/WebBrowserTextureResource.cpp
+++ b/Engine/Source/Runtime/WebBrowserTexture/Private/WebBrowserTextureResource.cpp
@@ -181,10 +181,10 @@ void FWebBrowserTextureResource::ClearTexture(const FLinearColor& ClearColor)
 	// draw the clear color
 	FRHICommandListImmediate& CommandList = FRHICommandListExecutor::GetImmediateCommandList();
 	{
-		FRHIRenderTargetView View = FRHIRenderTargetView(RenderTargetTextureRHI, ERenderTargetLoadAction::EClear);
-		FRHISetRenderTargetsInfo Info(1, &View, FRHIDepthRenderTargetView());
-		CommandList.SetRenderTargetsAndClear(Info);
-		CommandList.SetViewport(0, 0, 0.0f, RenderTargetTextureRHI->GetSizeX(), RenderTargetTextureRHI->GetSizeY(), 1.0f);
+		FRHIRenderPassInfo RPInfo(RenderTargetTextureRHI, ERenderTargetActions::Clear_Store);
+		CommandList.BeginRenderPass(RPInfo, TEXT("ClearTexture"));
+		CommandList.EndRenderPass();
+
 		CommandList.TransitionResource(EResourceTransitionAccess::EReadable, RenderTargetTextureRHI);
 	}
 
