Fix memory leak in Android 3D WebBrowser using Vulkan RHI

#android

diff --git a/Engine/Source/Runtime/WebBrowserTexture/Public/WebBrowserTextureSample.h b/Engine/Source/Runtime/WebBrowserTexture/Public/WebBrowserTextureSample.h
--- a/Engine/Source/Runtime/WebBrowserTexture/Public/WebBrowserTextureSample.h
+++ b/Engine/Source/Runtime/WebBrowserTexture/Public/WebBrowserTextureSample.h
@@ -59,8 +59,6 @@ public:
 	*/
 	bool Initialize(const FIntPoint& InDim)
 	{
-		BufferSize = 0;
-
 		if (InDim.GetMin() <= 0)
 		{ 
 			return false;
