Deprecating the functions RHICreateTextureReference and RHIUpdateTextureReference to force callers to provide command lists.
Adding RHIClearTextureReference as a temp workaround for ReleaseRHI() methods not having a command list.

diff --git a/Engine/Plugins/Experimental/PixelStreamingPlayer/Source/PixelStreamingPlayer/Private/PixelStreamingMediaTextureResource.cpp b/Engine/Plugins/Experimental/PixelStreamingPlayer/Source/PixelStreamingPlayer/Private/PixelStreamingMediaTextureResource.cpp
--- a/Engine/Plugins/Experimental/PixelStreamingPlayer/Source/PixelStreamingPlayer/Private/PixelStreamingMediaTextureResource.cpp
+++ b/Engine/Plugins/Experimental/PixelStreamingPlayer/Source/PixelStreamingPlayer/Private/PixelStreamingMediaTextureResource.cpp
@@ -26,7 +26,7 @@ void FPixelStreamingMediaTextureResource::ReleaseRHI()
 
 	if (MediaTexture != nullptr)
 	{
-		RHIUpdateTextureReference(MediaTexture->TextureReference.TextureReferenceRHI, nullptr);
+		RHIClearTextureReference(MediaTexture->TextureReference.TextureReferenceRHI);
 	}
 }
 
diff --git a/Engine/Plugins/Media/PixelStreaming2/Source/PixelStreaming2/Private/Blueprints/PixelStreaming2MediaTexture.cpp b/Engine/Plugins/Media/PixelStreaming2/Source/PixelStreaming2/Private/Blueprints/PixelStreaming2MediaTexture.cpp
--- a/Engine/Plugins/Media/PixelStreaming2/Source/PixelStreaming2/Private/Blueprints/PixelStreaming2MediaTexture.cpp
+++ b/Engine/Plugins/Media/PixelStreaming2/Source/PixelStreaming2/Private/Blueprints/PixelStreaming2MediaTexture.cpp
@@ -49,7 +49,7 @@ namespace UE::PixelStreaming2
 			TStrongObjectPtr<UPixelStreaming2MediaTexture> PinnedMediaTexture = MediaTexture.Pin();
 			if (PinnedMediaTexture.IsValid())
 			{
-				RHIUpdateTextureReference(PinnedMediaTexture->TextureReference.TextureReferenceRHI, nullptr);
+				RHIClearTextureReference(PinnedMediaTexture->TextureReference.TextureReferenceRHI);
 			}
 		}
 
diff --git a/Engine/Plugins/PCG/Source/PCGEditor/Private/DataVisualizations/PCGVisualizationTexture2D.cpp b/Engine/Plugins/PCG/Source/PCGEditor/Private/DataVisualizations/PCGVisualizationTexture2D.cpp
--- a/Engine/Plugins/PCG/Source/PCGEditor/Private/DataVisualizations/PCGVisualizationTexture2D.cpp
+++ b/Engine/Plugins/PCG/Source/PCGEditor/Private/DataVisualizations/PCGVisualizationTexture2D.cpp
@@ -25,7 +25,7 @@ namespace PCGVisTexture2DHelpers
 			FRHITextureCreateDesc::Create2D(TEXT("PCGDummyTexture"), 1, 1, PF_G8)
 			.SetFlags(ETextureCreateFlags::ShaderResource);
 
-		FTextureRHIRef DummyTexture = RHICreateTexture(DummyTextureDesc);
+		FTextureRHIRef DummyTexture = RHICmdList.CreateTexture(DummyTextureDesc);
 
 		{
 			uint32 DestStride;
@@ -90,7 +90,7 @@ void FPCGVisualizationTexture2DResource::InitRHI(FRHICommandListBase& RHICmdList
 			.SetFlags(SourceDesc.Flags)
 			.SetInitialState(ERHIAccess::CopyDest);
 
-		TextureRHI = RHICreateTexture(CreateDesc);
+		TextureRHI = RHICmdList.CreateTexture(CreateDesc);
 
 		// Prepare to copy the source texture.
 		RHICmdList.Transition(FRHITransitionInfo(SourceTextureRHI, ERHIAccess::Unknown, ERHIAccess::CopySrc));
@@ -110,14 +110,14 @@ void FPCGVisualizationTexture2DResource::InitRHI(FRHICommandListBase& RHICmdList
 
 	SamplerStateRHI = TStaticSamplerState<SF_Bilinear, AM_Clamp, AM_Clamp, AM_Clamp>::GetRHI();
 
-	RHIUpdateTextureReference(Owner->TextureReference.TextureReferenceRHI, TextureRHI);
+	RHICmdList.UpdateTextureReference(Owner->TextureReference.TextureReferenceRHI, TextureRHI);
 }
 
 void FPCGVisualizationTexture2DResource::ReleaseRHI()
 {
 	if (Owner)
 	{
-		RHIUpdateTextureReference(Owner->TextureReference.TextureReferenceRHI, nullptr);
+		RHIClearTextureReference(Owner->TextureReference.TextureReferenceRHI);
 	}
 
 	FTextureResource::ReleaseRHI();
diff --git a/Engine/Plugins/Runtime/AR/AppleAR/AppleARKit/Source/AppleARKit/Private/AppleARKitTextures.cpp b/Engine/Plugins/Runtime/AR/AppleAR/AppleARKit/Source/AppleARKit/Private/AppleARKitTextures.cpp
--- a/Engine/Plugins/Runtime/AR/AppleAR/AppleARKit/Source/AppleARKit/Private/AppleARKitTextures.cpp
+++ b/Engine/Plugins/Runtime/AR/AppleAR/AppleARKit/Source/AppleARKit/Private/AppleARKitTextures.cpp
@@ -162,14 +162,14 @@ public:
 			FRHITextureCreateDesc::Create2D(TEXT("FARKitTextureResource"), Size, PF_B8G8R8A8)
 			.SetFlags(ETextureCreateFlags::ShaderResource);
 		
-		TextureRHI = RHICreateTexture(Desc);
+		TextureRHI = RHICmdList.CreateTexture(Desc);
 		
 		OnTextureUpdated(RHICmdList);
 	}
 	
 	virtual void ReleaseRHI() override
 	{
-		RHIUpdateTextureReference(Owner->TextureReference.TextureReferenceRHI, nullptr);
+		RHIClearTextureReference(Owner->TextureReference.TextureReferenceRHI);
 		FTextureResource::ReleaseRHI();
 	}
 	
@@ -611,7 +611,7 @@ public:
 	
 	virtual void ReleaseRHI() override
 	{
-		RHIUpdateTextureReference(Owner->TextureReference.TextureReferenceRHI, nullptr);
+		RHIClearTextureReference(Owner->TextureReference.TextureReferenceRHI);
 		EnvCubemapTextureRHIRef.SafeRelease();
 		FTextureResource::ReleaseRHI();
 		FExternalTextureRegistry::Get().UnregisterExternalTexture(Owner->ExternalTextureGuid);
@@ -803,7 +803,7 @@ public:
 	
 	virtual void ReleaseRHI() override
 	{
-		RHIUpdateTextureReference(Owner->TextureReference.TextureReferenceRHI, nullptr);
+		RHIClearTextureReference(Owner->TextureReference.TextureReferenceRHI);
 		DecodedTextureRef.SafeRelease();
 		FTextureResource::ReleaseRHI();
 	}
diff --git a/Engine/Plugins/Runtime/AR/Google/GoogleARCore/Source/GoogleARCoreBase/Private/GoogleARCoreTexture.cpp b/Engine/Plugins/Runtime/AR/Google/GoogleARCore/Source/GoogleARCoreBase/Private/GoogleARCoreTexture.cpp
--- a/Engine/Plugins/Runtime/AR/Google/GoogleARCore/Source/GoogleARCoreBase/Private/GoogleARCoreTexture.cpp
+++ b/Engine/Plugins/Runtime/AR/Google/GoogleARCore/Source/GoogleARCoreBase/Private/GoogleARCoreTexture.cpp
@@ -159,7 +159,7 @@ public:
 	virtual void ReleaseRHI() override
 	{
 		FTextureResource::ReleaseRHI();
-		RHIUpdateTextureReference(Owner->TextureReference.TextureReferenceRHI, nullptr);
+		RHIClearTextureReference(Owner->TextureReference.TextureReferenceRHI);
 	}
 	
 	virtual uint32 GetSizeX() const override
diff --git a/Engine/Plugins/VirtualProduction/RemoteControl/Source/RemoteControlLogic/Private/Action/Bind/RCExternalTexture.cpp b/Engine/Plugins/VirtualProduction/RemoteControl/Source/RemoteControlLogic/Private/Action/Bind/RCExternalTexture.cpp
--- a/Engine/Plugins/VirtualProduction/RemoteControl/Source/RemoteControlLogic/Private/Action/Bind/RCExternalTexture.cpp
+++ b/Engine/Plugins/VirtualProduction/RemoteControl/Source/RemoteControlLogic/Private/Action/Bind/RCExternalTexture.cpp
@@ -92,7 +92,7 @@ public:
 		// Fast path: texture already cached.
 		if (CacheEntry && CacheEntry->TextureRHI.IsValid())
 		{
-			UpdateRHIFromCacheEntry();
+			UpdateRHIFromCacheEntry(RHICmdList);
 			return;
 		}
 		
@@ -107,22 +107,22 @@ public:
 		FTexture2DDynamicResource::ReleaseRHI();
 	}
 
-	void SetCacheEntry(const FString& InNewPath, const TSharedPtr<FResourceCacheEntry>& InNewEntry)
+	void SetCacheEntry(FRHICommandListBase& RHICmdList, const FString& InNewPath, const TSharedPtr<FResourceCacheEntry>& InNewEntry)
 	{
 		ReleaseRHI();
 
 		Path = InNewPath;
 		CacheEntry = InNewEntry;
-		UpdateRHIFromCacheEntry();
+		UpdateRHIFromCacheEntry(RHICmdList);
 	}
 
-	void UpdateRHIFromCacheEntry()
+	void UpdateRHIFromCacheEntry(FRHICommandListBase& RHICmdList)
 	{
 		if (CacheEntry)
 		{
 			TextureRHI = CacheEntry->TextureRHI;
 			SamplerStateRHI = CacheEntry->SamplerStateRHI;
-			RHIUpdateTextureReference(RCTextureOwner->TextureReference.TextureReferenceRHI,TextureRHI);
+			RHICmdList.UpdateTextureReference(RCTextureOwner->TextureReference.TextureReferenceRHI,TextureRHI);
 		}
 	}
 
@@ -181,7 +181,7 @@ void URCExternalTexture::LoadFromPath(const FString& InPath)
 			ENQUEUE_RENDER_COMMAND(FUpdateSharedTextureEntry)(
 				[TextureResource, ExistingEntry, NewPath=Path](FRHICommandListImmediate& RHICmdList)
 				{
-					TextureResource->SetCacheEntry(NewPath, ExistingEntry);
+					TextureResource->SetCacheEntry(RHICmdList, NewPath, ExistingEntry);
 				});
 		}
 		return;
diff --git a/Engine/Source/Runtime/Engine/Private/Rendering/StreamableTextureResource.cpp b/Engine/Source/Runtime/Engine/Private/Rendering/StreamableTextureResource.cpp
--- a/Engine/Source/Runtime/Engine/Private/Rendering/StreamableTextureResource.cpp
+++ b/Engine/Source/Runtime/Engine/Private/Rendering/StreamableTextureResource.cpp
@@ -238,14 +238,14 @@ void FStreamableTextureResource::ReleaseRHI()
 
 	if (ensure(TextureReferenceRHI.IsValid()))
 	{
-		RHIUpdateTextureReference(TextureReferenceRHI, nullptr);
+		RHIClearTextureReference(TextureReferenceRHI);
 	}
 
 	TextureRHI.SafeRelease();
 	FTextureResource::ReleaseRHI();
 }
 
-void FStreamableTextureResource::FinalizeStreaming(FRHITexture* InTextureRHI)
+void FStreamableTextureResource::FinalizeStreaming(FRHICommandListBase& RHICmdList, FRHITexture* InTextureRHI)
 {
 	checkSlow(IsInRenderingThread() && InTextureRHI);
 
@@ -266,11 +266,16 @@ void FStreamableTextureResource::FinalizeStreaming(FRHITexture* InTextureRHI)
 	TextureRHI->SetOwnerName(GetOwnerName());
 	if (ensure(TextureReferenceRHI.IsValid()))
 	{
-		RHIUpdateTextureReference(TextureReferenceRHI, TextureRHI);
+		RHICmdList.UpdateTextureReference(TextureReferenceRHI, TextureRHI);
 	}
 	State.NumResidentLODs = State.NumRequestedLODs;
 }
 
+void FStreamableTextureResource::FinalizeStreaming(FRHITexture* InTextureRHI)
+{
+	FinalizeStreaming(FRHICommandListImmediate::Get(), InTextureRHI);
+}
+
 // Recreate the sampler states (used when updating mip map lod bias offset)
 void FStreamableTextureResource::RefreshSamplerStates()
 {
diff --git a/Engine/Source/Runtime/Engine/Private/Rendering/Texture2DArrayResource.cpp b/Engine/Source/Runtime/Engine/Private/Rendering/Texture2DArrayResource.cpp
--- a/Engine/Source/Runtime/Engine/Private/Rendering/Texture2DArrayResource.cpp
+++ b/Engine/Source/Runtime/Engine/Private/Rendering/Texture2DArrayResource.cpp
@@ -117,7 +117,7 @@ void FTexture2DArrayResource::InitRHI(FRHICommandListBase& RHICmdList)
 	if (ProxiedResource)
 	{
 		TextureRHI = ProxiedResource->TextureRHI;
-		RHIUpdateTextureReference(TextureReferenceRHI, TextureRHI);
+		RHICmdList.UpdateTextureReference(TextureReferenceRHI, TextureRHI);
 		SamplerStateRHI = ProxiedResource->SamplerStateRHI;
 		DeferredPassSamplerStateRHI = ProxiedResource->DeferredPassSamplerStateRHI;
 	}
diff --git a/Engine/Source/Runtime/Engine/Private/Rendering/Texture3DResource.cpp b/Engine/Source/Runtime/Engine/Private/Rendering/Texture3DResource.cpp
--- a/Engine/Source/Runtime/Engine/Private/Rendering/Texture3DResource.cpp
+++ b/Engine/Source/Runtime/Engine/Private/Rendering/Texture3DResource.cpp
@@ -218,7 +218,7 @@ void FTexture3DResource::InitRHI(FRHICommandListBase& RHICmdList)
 	if (ProxiedResource)
 	{
 		TextureRHI = ProxiedResource->TextureRHI;
-		RHIUpdateTextureReference(TextureReferenceRHI, TextureRHI);
+		RHICmdList.UpdateTextureReference(TextureReferenceRHI, TextureRHI);
 		SamplerStateRHI = ProxiedResource->SamplerStateRHI;
 		DeferredPassSamplerStateRHI = ProxiedResource->DeferredPassSamplerStateRHI;
 	}
diff --git a/Engine/Source/Runtime/Engine/Private/Streaming/Texture2DUpdate.cpp b/Engine/Source/Runtime/Engine/Private/Streaming/Texture2DUpdate.cpp
--- a/Engine/Source/Runtime/Engine/Private/Streaming/Texture2DUpdate.cpp
+++ b/Engine/Source/Runtime/Engine/Private/Streaming/Texture2DUpdate.cpp
@@ -86,6 +86,8 @@ void FTexture2DUpdate::DoConvertToVirtualWithNewMips(const FContext& Context)
 
 	if (!IsCancelled() && Context.Resource)
 	{
+		FRHICommandListImmediate& RHICmdList = FRHICommandListImmediate::Get();
+
 		// If the texture is not virtual, then make it virtual immediately.
 		if (!Context.Resource->IsTextureRHIPartiallyResident())
 		{
@@ -106,13 +108,13 @@ void FTexture2DUpdate::DoConvertToVirtualWithNewMips(const FContext& Context)
 
 			Desc.DetermineInititialState();
 
-			IntermediateTextureRHI = RHICreateTexture(Desc);
+			IntermediateTextureRHI = RHICmdList.CreateTexture(Desc);
 
-			RHIVirtualTextureSetFirstMipInMemory(IntermediateTextureRHI, CurrentFirstLODIdx);
-			RHIVirtualTextureSetFirstMipVisible(IntermediateTextureRHI, CurrentFirstLODIdx);
+			RHICmdList.VirtualTextureSetFirstMipInMemory(IntermediateTextureRHI, CurrentFirstLODIdx);
+			RHICmdList.VirtualTextureSetFirstMipVisible(IntermediateTextureRHI, CurrentFirstLODIdx);
 
 			UE::RHI::CopySharedMips_AssumeSRVMaskState(
-				FRHICommandListExecutor::GetImmediateCommandList(),
+				RHICmdList,
 				Context.Resource->GetTexture2DRHI(),
 				IntermediateTextureRHI);
 		}
@@ -121,7 +123,7 @@ void FTexture2DUpdate::DoConvertToVirtualWithNewMips(const FContext& Context)
 			// Otherwise the current texture is already virtual and we can update it directly.
 			IntermediateTextureRHI = Context.Resource->GetTexture2DRHI();
 		}
-		RHIVirtualTextureSetFirstMipInMemory(IntermediateTextureRHI, PendingFirstLODIdx);
+		RHICmdList.VirtualTextureSetFirstMipInMemory(IntermediateTextureRHI, PendingFirstLODIdx);
 	}
 }
 
@@ -134,6 +136,8 @@ bool FTexture2DUpdate::DoConvertToNonVirtual(const FContext& Context)
 	{
 		if (Context.Resource->IsTextureRHIPartiallyResident())
 		{
+			FRHICommandListImmediate& RHICmdList = FRHICommandListImmediate::Get();
+
 			const FTexture2DMipMap& PendingFirstMipMap = *Context.MipsView[PendingFirstLODIdx];
 
 			ensure(!IntermediateTextureRHI);
@@ -150,10 +154,10 @@ bool FTexture2DUpdate::DoConvertToNonVirtual(const FContext& Context)
 
 			Desc.DetermineInititialState();
 
-			IntermediateTextureRHI = RHICreateTexture(Desc);
+			IntermediateTextureRHI = RHICmdList.CreateTexture(Desc);
 
 			UE::RHI::CopySharedMips_AssumeSRVMaskState(
-				FRHICommandListExecutor::GetImmediateCommandList(),
+				RHICmdList,
 				Context.Resource->GetTexture2DRHI(),
 				IntermediateTextureRHI);
 
@@ -171,7 +175,7 @@ void FTexture2DUpdate::DoFinishUpdate(const FContext& Context)
 	{
 		if (!IsCancelled())
 		{
-			Context.Resource->FinalizeStreaming(IntermediateTextureRHI);
+			Context.Resource->FinalizeStreaming(FRHICommandListImmediate::Get(), IntermediateTextureRHI);
 			MarkAsSuccessfullyFinished();
 		}
 		IntermediateTextureRHI.SafeRelease();
diff --git a/Engine/Source/Runtime/Engine/Private/Streaming/TextureMipAllocator.cpp b/Engine/Source/Runtime/Engine/Private/Streaming/TextureMipAllocator.cpp
--- a/Engine/Source/Runtime/Engine/Private/Streaming/TextureMipAllocator.cpp
+++ b/Engine/Source/Runtime/Engine/Private/Streaming/TextureMipAllocator.cpp
@@ -30,7 +30,7 @@ bool FTextureMipAllocator::FinalizeMips(const FTextureUpdateContext& Context, co
 	}
 
 	// Use the new texture resource for the texture asset, must run on the renderthread.
-	Context.Resource->FinalizeStreaming(IntermediateTextureRHI);
+	Context.Resource->FinalizeStreaming(FRHICommandListImmediate::Get(), IntermediateTextureRHI);
 
 	// No need for the intermediate texture anymore.
 	IntermediateTextureRHI.SafeRelease();
diff --git a/Engine/Source/Runtime/Engine/Private/Texture2D.cpp b/Engine/Source/Runtime/Engine/Private/Texture2D.cpp
--- a/Engine/Source/Runtime/Engine/Private/Texture2D.cpp
+++ b/Engine/Source/Runtime/Engine/Private/Texture2D.cpp
@@ -1755,7 +1755,7 @@ void FVirtualTexture2DResource::InitializeEditorResources(FRHICommandListBase& R
 					FRHITextureCreateDesc::Create2D(TEXT("VTEditorResource"), MipWidthInTiles * TileSizeInPixels, MipHeightInTiles * TileSizeInPixels, PixelFormat)
 					.SetFlags(ETextureCreateFlags::ShaderResource);
 
-				Texture2DRHI = RHICreateTexture(Desc);
+				Texture2DRHI = RHICmdList.CreateTexture(Desc);
 			}
 			
 			TRefCountPtr<IPooledRenderTarget> PooledRenderTarget = CreateRenderTarget(Texture2DRHI, TEXT("VTEditorResource"));
@@ -1810,7 +1810,7 @@ void FVirtualTexture2DResource::InitializeEditorResources(FRHICommandListBase& R
 		// Only update texture reference if necessary.
 		if (TextureRHI != TextureReferenceRHI->GetReferencedTexture())
 		{
-			RHIUpdateTextureReference(TextureReferenceRHI, TextureRHI);
+			RHICmdList.UpdateTextureReference(TextureReferenceRHI, TextureRHI);
 		}
 	}
 }
diff --git a/Engine/Source/Runtime/Engine/Private/Texture2DDynamic.cpp b/Engine/Source/Runtime/Engine/Private/Texture2DDynamic.cpp
--- a/Engine/Source/Runtime/Engine/Private/Texture2DDynamic.cpp
+++ b/Engine/Source/Runtime/Engine/Private/Texture2DDynamic.cpp
@@ -36,7 +36,7 @@ uint32 FTexture2DDynamicResource::GetSizeY() const
 }
 
 /** Called when the resource is initialized. This is only called by the rendering thread. */
-void FTexture2DDynamicResource::InitRHI(FRHICommandListBase&)
+void FTexture2DDynamicResource::InitRHI(FRHICommandListBase& RHICmdList)
 {
 	// Create the sampler state RHI resource.
 	ESamplerAddressMode SamplerAddressMode = Owner->SamplerAddressMode;
@@ -69,17 +69,17 @@ void FTexture2DDynamicResource::InitRHI(FRHICommandListBase&)
 		Desc.AddFlags(ETextureCreateFlags::NoTiling);
 	}
 
-	Texture2DRHI = RHICreateTexture(Desc);
+	Texture2DRHI = RHICmdList.CreateTexture(Desc);
 
 	TextureRHI = Texture2DRHI;
 	TextureRHI->SetName(Owner->GetFName());
-	RHIUpdateTextureReference(Owner->TextureReference.TextureReferenceRHI,TextureRHI);
+	RHICmdList.UpdateTextureReference(Owner->TextureReference.TextureReferenceRHI,TextureRHI);
 }
 
 /** Called when the resource is released. This is only called by the rendering thread. */
 void FTexture2DDynamicResource::ReleaseRHI()
 {
-	RHIUpdateTextureReference(Owner->TextureReference.TextureReferenceRHI, nullptr);
+	RHIClearTextureReference(Owner->TextureReference.TextureReferenceRHI);
 	FTextureResource::ReleaseRHI();
 	Texture2DRHI.SafeRelease();
 }
diff --git a/Engine/Source/Runtime/Engine/Private/TextureCube.cpp b/Engine/Source/Runtime/Engine/Private/TextureCube.cpp
--- a/Engine/Source/Runtime/Engine/Private/TextureCube.cpp
+++ b/Engine/Source/Runtime/Engine/Private/TextureCube.cpp
@@ -446,7 +446,7 @@ public:
 			TextureCubeRHI = ProxiedResource->GetTextureCubeRHI();
 			TextureRHI = TextureCubeRHI;
 			SamplerStateRHI = ProxiedResource->SamplerStateRHI;
-			RHIUpdateTextureReference(Owner->TextureReference.TextureReferenceRHI, TextureRHI);
+			RHICmdList.UpdateTextureReference(Owner->TextureReference.TextureReferenceRHI, TextureRHI);
 			return;
 		}
 
@@ -468,10 +468,10 @@ public:
 			.SetClassName(ClassName)
 			.SetOwnerName(GetOwnerName());
 
-		TextureCubeRHI = RHICreateTexture(Desc);
+		TextureCubeRHI = RHICmdList.CreateTexture(Desc);
 
 		TextureRHI = TextureCubeRHI;
-		RHIUpdateTextureReference(Owner->TextureReference.TextureReferenceRHI,TextureRHI);
+		RHICmdList.UpdateTextureReference(Owner->TextureReference.TextureReferenceRHI, TextureRHI);
 
 		// Read the mip-levels into the RHI texture.
 		int32 NumMips = Owner->GetNumMips();
@@ -507,7 +507,7 @@ public:
 	{
 		DEC_DWORD_STAT_BY( STAT_TextureMemory, TextureSize );
 		DEC_DWORD_STAT_FNAME_BY( LODGroupStatName, TextureSize );
-		RHIUpdateTextureReference(Owner->TextureReference.TextureReferenceRHI, nullptr);
+		RHIClearTextureReference(Owner->TextureReference.TextureReferenceRHI);
 		TextureCubeRHI.SafeRelease();
 		FTextureResource::ReleaseRHI();
 	}
diff --git a/Engine/Source/Runtime/Engine/Private/TextureCubeArray.cpp b/Engine/Source/Runtime/Engine/Private/TextureCubeArray.cpp
--- a/Engine/Source/Runtime/Engine/Private/TextureCubeArray.cpp
+++ b/Engine/Source/Runtime/Engine/Private/TextureCubeArray.cpp
@@ -118,7 +118,7 @@ public:
 			TextureCubeRHI = ProxiedResource->GetTextureCubeRHI();
 			TextureRHI = TextureCubeRHI;
 			SamplerStateRHI = ProxiedResource->SamplerStateRHI;
-			RHIUpdateTextureReference(Owner->TextureReference.TextureReferenceRHI, TextureRHI);
+			RHICmdList.UpdateTextureReference(Owner->TextureReference.TextureReferenceRHI, TextureRHI);
 			return;
 		}
 
@@ -145,10 +145,10 @@ public:
 			.SetClassName(ClassName)
 			.SetOwnerName(GetOwnerName());
 
-		TextureCubeRHI = RHICreateTexture(Desc);
+		TextureCubeRHI = RHICmdList.CreateTexture(Desc);
 		TextureRHI = TextureCubeRHI;
 
-		RHIUpdateTextureReference(Owner->TextureReference.TextureReferenceRHI, TextureRHI);
+		RHICmdList.UpdateTextureReference(Owner->TextureReference.TextureReferenceRHI, TextureRHI);
 
 		// Read the mip-levels into the RHI texture.
 		const int32 FirstMipTailIndex = Owner->GetNumMips() - FMath::Max(1, Owner->GetPlatformData()->GetNumMipsInTail());
@@ -183,7 +183,7 @@ public:
 	{
 		DEC_DWORD_STAT_BY(STAT_TextureMemory, TextureSize);
 		DEC_DWORD_STAT_FNAME_BY(LODGroupStatName, TextureSize);
-		RHIUpdateTextureReference(Owner->TextureReference.TextureReferenceRHI, nullptr);
+		RHIClearTextureReference(Owner->TextureReference.TextureReferenceRHI);
 		TextureCubeRHI.SafeRelease();
 		FTextureResource::ReleaseRHI();
 	}
diff --git a/Engine/Source/Runtime/Engine/Private/TextureRenderTarget2D.cpp b/Engine/Source/Runtime/Engine/Private/TextureRenderTarget2D.cpp
--- a/Engine/Source/Runtime/Engine/Private/TextureRenderTarget2D.cpp
+++ b/Engine/Source/Runtime/Engine/Private/TextureRenderTarget2D.cpp
@@ -545,7 +545,7 @@ void FTextureRenderTarget2DResource::InitRHI(FRHICommandListBase& RHICmdList)
 			.SetClassName(ClassName)
 			.SetOwnerName(GetOwnerName());
 
-		TextureRHI = RenderTargetTextureRHI = RHICreateTexture(Desc);
+		TextureRHI = RenderTargetTextureRHI = RHICmdList.CreateTexture(Desc);
 
 		if (NumSamples > 1)
 		{
@@ -556,13 +556,13 @@ void FTextureRenderTarget2DResource::InitRHI(FRHICommandListBase& RHICmdList)
 			Desc.SetFlags(ResolveTexCreateFlags);
 			Desc.SetNumSamples(1);
 
-			TextureRHI = RHICreateTexture(Desc);
+			TextureRHI = RHICmdList.CreateTexture(Desc);
 		}
 		else if (Owner->bNeedsTwoCopies)
 		{
 			Desc.SetFlags(TexCreateFlags | ETextureCreateFlags::ShaderResource);
 
-			TextureRHI = RHICreateTexture(Desc);
+			TextureRHI = RHICmdList.CreateTexture(Desc);
 		}
 
 		if (EnumHasAnyFlags(TexCreateFlags, ETextureCreateFlags::UAV))
@@ -571,9 +571,7 @@ void FTextureRenderTarget2DResource::InitRHI(FRHICommandListBase& RHICmdList)
 		}
 
 		SetGPUMask(FRHIGPUMask::All());
-		RHIUpdateTextureReference(Owner->TextureReference.TextureReferenceRHI, TextureRHI);
-
-		TextureRHI->SetOwnerName(GetOwnerName());
+		RHICmdList.UpdateTextureReference(Owner->TextureReference.TextureReferenceRHI, TextureRHI);
 
 		AddToDeferredUpdateList(true);
 	}
@@ -602,7 +600,7 @@ void FTextureRenderTarget2DResource::ReleaseRHI()
 	// release the FTexture RHI resources here as well
 	FTexture::ReleaseRHI();
 
-	RHIUpdateTextureReference(Owner->TextureReference.TextureReferenceRHI, nullptr);
+	RHIClearTextureReference(Owner->TextureReference.TextureReferenceRHI);
 	RenderTargetTextureRHI.SafeRelease();
 	MipGenerationCache.SafeRelease();
 
diff --git a/Engine/Source/Runtime/Engine/Private/TextureRenderTarget2DArray.cpp b/Engine/Source/Runtime/Engine/Private/TextureRenderTarget2DArray.cpp
--- a/Engine/Source/Runtime/Engine/Private/TextureRenderTarget2DArray.cpp
+++ b/Engine/Source/Runtime/Engine/Private/TextureRenderTarget2DArray.cpp
@@ -255,7 +255,7 @@ void FTextureRenderTarget2DArrayResource::InitRHI(FRHICommandListBase& RHICmdLis
 				.SetClearValue(FClearValueBinding(Owner->ClearColor))
 				.SetInitialState(ERHIAccess::SRVMask);
 
-			TextureRHI = RHICreateTexture(Desc);
+			TextureRHI = RHICmdList.CreateTexture(Desc);
 		}
 
 		if (EnumHasAnyFlags(TexCreateFlags, ETextureCreateFlags::UAV))
@@ -263,7 +263,7 @@ void FTextureRenderTarget2DArrayResource::InitRHI(FRHICommandListBase& RHICmdLis
 			UnorderedAccessViewRHI = RHICmdList.CreateUnorderedAccessView(TextureRHI, FRHIViewDesc::CreateTextureUAV().SetDimensionFromTexture(TextureRHI));
 		}
 
-		RHIUpdateTextureReference(Owner->TextureReference.TextureReferenceRHI, TextureRHI);
+		RHICmdList.UpdateTextureReference(Owner->TextureReference.TextureReferenceRHI, TextureRHI);
 		RenderTargetTextureRHI = TextureRHI;
 
 		AddToDeferredUpdateList(true);
@@ -293,7 +293,7 @@ void FTextureRenderTarget2DArrayResource::ReleaseRHI()
 	// release the FTexture RHI resources here as well
 	FTexture::ReleaseRHI();
 
-	RHIUpdateTextureReference(Owner->TextureReference.TextureReferenceRHI, nullptr);
+	RHIClearTextureReference(Owner->TextureReference.TextureReferenceRHI);
 	RenderTargetTextureRHI.SafeRelease();
 
 	// remove from global list of deferred clears
diff --git a/Engine/Source/Runtime/Engine/Private/TextureRenderTargetCube.cpp b/Engine/Source/Runtime/Engine/Private/TextureRenderTargetCube.cpp
--- a/Engine/Source/Runtime/Engine/Private/TextureRenderTargetCube.cpp
+++ b/Engine/Source/Runtime/Engine/Private/TextureRenderTargetCube.cpp
@@ -309,7 +309,7 @@ void FTextureRenderTargetCubeResource::InitRHI(FRHICommandListBase& RHICmdList)
 				.SetClearValue(FClearValueBinding(Owner->ClearColor))
 				.SetInitialState(ERHIAccess::SRVMask);
 
-			TextureRHI = RHICreateTexture(Desc);
+			TextureRHI = RHICmdList.CreateTexture(Desc);
 		}
 
 		SetGPUMask(FRHIGPUMask::All());
@@ -319,7 +319,7 @@ void FTextureRenderTargetCubeResource::InitRHI(FRHICommandListBase& RHICmdList)
 			UnorderedAccessViewRHI = RHICmdList.CreateUnorderedAccessView(TextureRHI, FRHIViewDesc::CreateTextureUAV().SetDimensionFromTexture(TextureRHI));
 		}
 
-		RHIUpdateTextureReference(Owner->TextureReference.TextureReferenceRHI, TextureRHI);
+		RHICmdList.UpdateTextureReference(Owner->TextureReference.TextureReferenceRHI, TextureRHI);
 
 		// Create the RHI target surface used for rendering to
 		{
@@ -332,7 +332,7 @@ void FTextureRenderTargetCubeResource::InitRHI(FRHICommandListBase& RHICmdList)
 				.SetClearValue(FClearValueBinding(Owner->ClearColor))
 				.SetInitialState(ERHIAccess::SRVMask);
 
-			RenderTargetTextureRHI = RHICreateTexture(Desc);
+			RenderTargetTextureRHI = RHICmdList.CreateTexture(Desc);
 		}
 
 		AddToDeferredUpdateList(true);
@@ -359,7 +359,7 @@ void FTextureRenderTargetCubeResource::ReleaseRHI()
 	// release the FTexture RHI resources here as well
 	FTexture::ReleaseRHI();
 
-	RHIUpdateTextureReference(Owner->TextureReference.TextureReferenceRHI, nullptr);
+	RHIClearTextureReference(Owner->TextureReference.TextureReferenceRHI);
 	RenderTargetTextureRHI.SafeRelease();
 
 	// remove from global list of deferred clears
diff --git a/Engine/Source/Runtime/Engine/Private/TextureRenderTargetVolume.cpp b/Engine/Source/Runtime/Engine/Private/TextureRenderTargetVolume.cpp
--- a/Engine/Source/Runtime/Engine/Private/TextureRenderTargetVolume.cpp
+++ b/Engine/Source/Runtime/Engine/Private/TextureRenderTargetVolume.cpp
@@ -259,7 +259,7 @@ void FTextureRenderTargetVolumeResource::InitRHI(FRHICommandListBase& RHICmdList
 				.SetFlags(TexCreateFlags | ETextureCreateFlags::RenderTargetable | ETextureCreateFlags::ShaderResource)
 				.SetClearValue(FClearValueBinding(Owner->ClearColor));
 
-			TextureRHI = RenderTargetTextureRHI = RHICreateTexture(Desc);
+			TextureRHI = RenderTargetTextureRHI = RHICmdList.CreateTexture(Desc);
 		}
 
 		if (EnumHasAnyFlags(TexCreateFlags, TexCreate_UAV))
@@ -267,7 +267,7 @@ void FTextureRenderTargetVolumeResource::InitRHI(FRHICommandListBase& RHICmdList
 			UnorderedAccessViewRHI = RHICmdList.CreateUnorderedAccessView(TextureRHI, FRHIViewDesc::CreateTextureUAV().SetDimensionFromTexture(TextureRHI));
 		}
 
-		RHIUpdateTextureReference(Owner->TextureReference.TextureReferenceRHI,TextureRHI);
+		RHICmdList.UpdateTextureReference(Owner->TextureReference.TextureReferenceRHI, TextureRHI);
 
 		AddToDeferredUpdateList(true);
 	}
@@ -293,7 +293,7 @@ void FTextureRenderTargetVolumeResource::ReleaseRHI()
 	// release the FTexture RHI resources here as well
 	FTexture::ReleaseRHI();
 
-	RHIUpdateTextureReference(Owner->TextureReference.TextureReferenceRHI, nullptr);
+	RHIClearTextureReference(Owner->TextureReference.TextureReferenceRHI);
 	RenderTargetTextureRHI.SafeRelease();
 
 	// remove from global list of deferred clears
diff --git a/Engine/Source/Runtime/Engine/Public/Rendering/StreamableTextureResource.h b/Engine/Source/Runtime/Engine/Public/Rendering/StreamableTextureResource.h
--- a/Engine/Source/Runtime/Engine/Public/Rendering/StreamableTextureResource.h
+++ b/Engine/Source/Runtime/Engine/Public/Rendering/StreamableTextureResource.h
@@ -48,6 +48,8 @@ public:
 		return State; 
 	}
 
+	void FinalizeStreaming(FRHICommandListBase& RHICmdList, FRHITexture* InTextureRHI);
+	UE_DEPRECATED(5.7, "A RHICmdList argument is now required.")
 	void FinalizeStreaming(FRHITexture* InTextureRHI);
 
 	// Return the platform data mips, zero based on the resource mip count and not to the asset mip count.
diff --git a/Engine/Source/Runtime/RHI/Public/RHICommandList.h b/Engine/Source/Runtime/RHI/Public/RHICommandList.h
--- a/Engine/Source/Runtime/RHI/Public/RHICommandList.h
+++ b/Engine/Source/Runtime/RHI/Public/RHICommandList.h
@@ -5532,16 +5532,24 @@ inline FRHICommandListImmediate& FRHICommandListExecutor::GetImmediateCommandLis
 	return GRHICommandList.CommandListImmediate;
 }
 
+UE_DEPRECATED(5.7, "RHICreateTextureReference with an implied immediate command list is deprecated")
 inline FTextureReferenceRHIRef RHICreateTextureReference(FRHITexture* InReferencedTexture = nullptr)
 {
 	return FRHICommandListImmediate::Get().CreateTextureReference(InReferencedTexture);
 }
 
+UE_DEPRECATED(5.7, "RHIUpdateTextureReference with an implied immediate command list is deprecated")
 inline void RHIUpdateTextureReference(FRHITextureReference* TextureRef, FRHITexture* NewTexture)
 {
 	FRHICommandListImmediate::Get().UpdateTextureReference(TextureRef, NewTexture);
 }
 
+// Temporary workaround until we figure out if we want to route a command list through ReleaseRHI()
+inline void RHIClearTextureReference(FRHITextureReference* TextureRef)
+{
+	FRHICommandListImmediate::Get().UpdateTextureReference(TextureRef, nullptr);
+}
+
 inline FTextureRHIRef RHICreateTexture(const FRHITextureCreateDesc& CreateDesc)
 {
 	return FRHICommandListImmediate::Get().CreateTexture(CreateDesc);
diff --git a/Engine/Source/Runtime/RenderCore/Private/RenderResource.cpp b/Engine/Source/Runtime/RenderCore/Private/RenderResource.cpp
--- a/Engine/Source/Runtime/RenderCore/Private/RenderResource.cpp
+++ b/Engine/Source/Runtime/RenderCore/Private/RenderResource.cpp
@@ -523,10 +523,10 @@ void FTextureReference::InvalidateLastRenderTime()
 	}
 }
 
-void FTextureReference::InitRHI(FRHICommandListBase&)
+void FTextureReference::InitRHI(FRHICommandListBase& RHICmdList)
 {
 	SCOPED_LOADTIMER(FTextureReference_InitRHI);
-	TextureReferenceRHI = RHICreateTextureReference();
+	TextureReferenceRHI = RHICmdList.CreateTextureReference();
 }
 	
 int32 GTextureReferenceRevertsLastRenderContainer = 1;
diff --git a/Engine/Source/Runtime/Renderer/Private/VT/VirtualTextureSpace.cpp b/Engine/Source/Runtime/Renderer/Private/VT/VirtualTextureSpace.cpp
--- a/Engine/Source/Runtime/Renderer/Private/VT/VirtualTextureSpace.cpp
+++ b/Engine/Source/Runtime/Renderer/Private/VT/VirtualTextureSpace.cpp
@@ -338,7 +338,7 @@ void FVirtualTextureSpace::AllocateTextures(FRDGBuilder& GraphBuilder)
 			TextureEntry.RenderTarget = GraphBuilder.ConvertToExternalTexture(DstTexture);
 			GraphBuilder.UseExternalAccessMode(DstTexture, ERHIAccess::SRVMask, ERHIPipeline::All);
 
-			RHIUpdateTextureReference(TextureEntry.TextureReferenceRHI, TextureEntry.RenderTarget->GetRHI());
+			GraphBuilder.RHICmdList.UpdateTextureReference(TextureEntry.TextureReferenceRHI, TextureEntry.RenderTarget->GetRHI());
 		}
 
 		bNeedToAllocatePageTable = false;
@@ -360,7 +360,7 @@ void FVirtualTextureSpace::AllocateTextures(FRDGBuilder& GraphBuilder)
 			AddClearUAVPass(GraphBuilder, GraphBuilder.CreateUAV(PageTableIndirectionTexture), FUintVector4(ForceInitToZero));
 			PageTableIndirection.RenderTarget = GraphBuilder.ConvertToExternalTexture(PageTableIndirectionTexture);
 			GraphBuilder.UseExternalAccessMode(PageTableIndirectionTexture, ERHIAccess::SRVMask, ERHIPipeline::All);
-			RHIUpdateTextureReference(PageTableIndirection.TextureReferenceRHI, PageTableIndirection.RenderTarget->GetRHI());
+			GraphBuilder.RHICmdList.UpdateTextureReference(PageTableIndirection.TextureReferenceRHI, PageTableIndirection.RenderTarget->GetRHI());
 		}
 
 		bNeedToAllocatePageTableIndirection = false;
diff --git a/Engine/Source/Runtime/WebBrowserTexture/Private/WebBrowserTexture.cpp b/Engine/Source/Runtime/WebBrowserTexture/Private/WebBrowserTexture.cpp
--- a/Engine/Source/Runtime/WebBrowserTexture/Private/WebBrowserTexture.cpp
+++ b/Engine/Source/Runtime/WebBrowserTexture/Private/WebBrowserTexture.cpp
@@ -133,7 +133,7 @@ void UWebBrowserTexture::TickResource(TSharedPtr<FWebBrowserTextureSample, ESPMo
 	ENQUEUE_RENDER_COMMAND(UWebBrowserTextureResourceRender)(
 		[ResourceParam, RenderParams](FRHICommandListImmediate& RHICmdList)
 		{
-			ResourceParam->Render(RenderParams);
+			ResourceParam->Render(RHICmdList, RenderParams);
 		});
 }
 
diff --git a/Engine/Source/Runtime/WebBrowserTexture/Private/WebBrowserTextureResource.cpp b/Engine/Source/Runtime/WebBrowserTexture/Private/WebBrowserTextureResource.cpp
--- a/Engine/Source/Runtime/WebBrowserTexture/Private/WebBrowserTextureResource.cpp
+++ b/Engine/Source/Runtime/WebBrowserTexture/Private/WebBrowserTextureResource.cpp
@@ -36,7 +36,7 @@ FWebBrowserTextureResource::FWebBrowserTextureResource(UWebBrowserTexture& InOwn
 /* FWebBrowserTextureResource interface
  *****************************************************************************/
 
-void FWebBrowserTextureResource::Render(const FRenderParams& Params)
+void FWebBrowserTextureResource::Render(FRHICommandList& RHICmdList, const FRenderParams& Params)
 {
 	check(IsInRenderingThread());
 
@@ -61,7 +61,7 @@ void FWebBrowserTextureResource::Render(const FRenderParams& Params)
 		check(Sample.IsValid());
 
 		// render the sample
-		CopySample(Sample, Params.ClearColor);
+		CopySample(RHICmdList, Sample, Params.ClearColor);
 
 		if (!GSupportsImageExternal && Params.PlayerGuid.IsValid())
 		{
@@ -71,7 +71,7 @@ void FWebBrowserTextureResource::Render(const FRenderParams& Params)
 	}
 	else if (!Cleared)
 	{
-		ClearTexture(Params.ClearColor);
+		ClearTexture(RHICmdList, Params.ClearColor);
 
 		if (!GSupportsImageExternal && Params.PlayerGuid.IsValid())
 		{
@@ -135,14 +135,15 @@ void FWebBrowserTextureResource::ReleaseRHI()
 	RenderTargetTextureRHI.SafeRelease();
 	TextureRHI.SafeRelease();
 
-	UpdateTextureReference(nullptr);
+	RHIClearTextureReference(Owner.TextureReference.TextureReferenceRHI);
+	OwnerDim = FIntPoint::ZeroValue;
 }
 
 
 /* FWebBrowserTextureResource implementation
  *****************************************************************************/
 
-void FWebBrowserTextureResource::ClearTexture(const FLinearColor& ClearColor)
+void FWebBrowserTextureResource::ClearTexture(FRHICommandList& RHICmdList, const FLinearColor& ClearColor)
 {
 	UE_LOG(LogWebBrowserTexture, VeryVerbose, TEXT("FWebBrowserTextureResource:ClearTexture"));
 	// create output render target if we don't have one yet
@@ -158,7 +159,7 @@ void FWebBrowserTextureResource::ClearTexture(const FLinearColor& ClearColor)
 			.SetInitialState(ERHIAccess::SRVMask)
 			.SetClearValue(FClearValueBinding(ClearColor));
 
-		OutputTarget = RHICreateTexture(Desc);
+		OutputTarget = RHICmdList.CreateTexture(Desc);
 
 		CurrentClearColor = ClearColor;
 		UpdateResourceSize();
@@ -166,21 +167,20 @@ void FWebBrowserTextureResource::ClearTexture(const FLinearColor& ClearColor)
 
 	if (RenderTargetTextureRHI != OutputTarget)
 	{
-		UpdateTextureReference(OutputTarget);
+		UpdateTextureReference(RHICmdList, OutputTarget);
 	}
 
 	// draw the clear color
-	FRHICommandListImmediate& CommandList = FRHICommandListExecutor::GetImmediateCommandList();
 	{
-		CommandList.Transition(FRHITransitionInfo(RenderTargetTextureRHI, ERHIAccess::Unknown, ERHIAccess::RTV));
-		ClearRenderTarget(CommandList, RenderTargetTextureRHI);
-		CommandList.Transition(FRHITransitionInfo(RenderTargetTextureRHI, ERHIAccess::RTV, ERHIAccess::SRVMask));
+		RHICmdList.Transition(FRHITransitionInfo(RenderTargetTextureRHI, ERHIAccess::Unknown, ERHIAccess::RTV));
+		ClearRenderTarget(RHICmdList, RenderTargetTextureRHI);
+		RHICmdList.Transition(FRHITransitionInfo(RenderTargetTextureRHI, ERHIAccess::RTV, ERHIAccess::SRVMask));
 	}
 
 	Cleared = true;
 }
 
-void FWebBrowserTextureResource::CopySample(const TSharedPtr<FWebBrowserTextureSample, ESPMode::ThreadSafe>& Sample, const FLinearColor& ClearColor)
+void FWebBrowserTextureResource::CopySample(FRHICommandListBase& RHICmdList, const TSharedPtr<FWebBrowserTextureSample, ESPMode::ThreadSafe>& Sample, const FLinearColor& ClearColor)
 {
 	FRHITexture* SampleTexture = Sample->GetTexture();
 	FRHITexture* SampleTexture2D = (SampleTexture != nullptr) ? SampleTexture->GetTexture2D() : nullptr;
@@ -194,7 +194,7 @@ void FWebBrowserTextureResource::CopySample(const TSharedPtr<FWebBrowserTextureS
 		if (TextureRHI != SampleTexture2D)
 		{
 			UE_LOG(LogWebBrowserTexture, VeryVerbose, TEXT("FWebBrowserTextureResource:CopySample 11"));
-			UpdateTextureReference(SampleTexture2D);
+			UpdateTextureReference(RHICmdList, SampleTexture2D);
 
 			OutputTarget.SafeRelease();
 			UpdateResourceSize();
@@ -219,7 +219,7 @@ void FWebBrowserTextureResource::CopySample(const TSharedPtr<FWebBrowserTextureS
 				.SetInitialState(ERHIAccess::SRVMask)
 				.SetClearValue(FClearValueBinding(ClearColor));
 
-			OutputTarget = RHICreateTexture(Desc);
+			OutputTarget = RHICmdList.CreateTexture(Desc);
 
 			CurrentClearColor = ClearColor;
 			UpdateResourceSize();
@@ -227,7 +227,7 @@ void FWebBrowserTextureResource::CopySample(const TSharedPtr<FWebBrowserTextureS
 
 		if (RenderTargetTextureRHI != OutputTarget)
 		{
-			UpdateTextureReference(OutputTarget);
+			UpdateTextureReference(RHICmdList, OutputTarget);
 		}
 
 		UE_LOG(LogWebBrowserTexture, VeryVerbose, TEXT("WebBrowserTextureResource:CopySample: %d x %d"), SampleDim.X, SampleDim.Y);
@@ -260,12 +260,12 @@ void FWebBrowserTextureResource::UpdateResourceSize()
 }
 
 
-void FWebBrowserTextureResource::UpdateTextureReference(FRHITexture* NewTexture)
+void FWebBrowserTextureResource::UpdateTextureReference(FRHICommandListBase& RHICmdList, FRHITexture* NewTexture)
 {
 	TextureRHI = NewTexture;
 	RenderTargetTextureRHI = NewTexture;
 
-	RHIUpdateTextureReference(Owner.TextureReference.TextureReferenceRHI, NewTexture);
+	RHICmdList.UpdateTextureReference(Owner.TextureReference.TextureReferenceRHI, NewTexture);
 
 	if (RenderTargetTextureRHI != nullptr)
 	{
diff --git a/Engine/Source/Runtime/WebBrowserTexture/Public/WebBrowserTextureResource.h b/Engine/Source/Runtime/WebBrowserTexture/Public/WebBrowserTextureResource.h
--- a/Engine/Source/Runtime/WebBrowserTexture/Public/WebBrowserTextureResource.h
+++ b/Engine/Source/Runtime/WebBrowserTexture/Public/WebBrowserTextureResource.h
@@ -63,7 +63,7 @@ public:
 	 *
 	 * @param Params Render parameters.
 	 */
-	void Render(const FRenderParams& Params);
+	void Render(FRHICommandList& RHICmdList, const FRenderParams& Params);
 
 public:
 
@@ -88,7 +88,7 @@ protected:
 	 *
 	 * @param ClearColor The clear color to use.
 	 */
-	void ClearTexture(const FLinearColor& ClearColor);
+	void ClearTexture(FRHICommandList& RHICmdList, const FLinearColor& ClearColor);
 
 	/**
 	 * Render the given texture sample by using it as or copying it to the render target.
@@ -97,7 +97,7 @@ protected:
 	 * @param ClearColor The clear color to use for the output texture.
 	 * @see ConvertSample
 	 */
-	void CopySample(const TSharedPtr<FWebBrowserTextureSample, ESPMode::ThreadSafe>& Sample, const FLinearColor& ClearColor);
+	void CopySample(FRHICommandListBase& RHICmdList, const TSharedPtr<FWebBrowserTextureSample, ESPMode::ThreadSafe>& Sample, const FLinearColor& ClearColor);
 
 	/** Calculates the current resource size and notifies the owner texture. */
 	void UpdateResourceSize();
@@ -107,7 +107,7 @@ protected:
 	 *
 	 * @param NewTexture The texture to set.
 	 */
-	void UpdateTextureReference(FRHITexture* NewTexture);
+	void UpdateTextureReference(FRHICommandListBase& RHICmdList, FRHITexture* NewTexture);
 
 private:
 
